ElectraCodecs: Added dictionary parameter to receive decoder information upon instance creation.

#rnx

diff --git a/Engine/Plugins/Media/AppleProResDecoderElectra/Source/AppleProResDecoderElectra/Private/ProResDecoder/Mac/ElectraMediaProResDecoder.cpp b/Engine/Plugins/Media/AppleProResDecoderElectra/Source/AppleProResDecoderElectra/Private/ProResDecoder/Mac/ElectraMediaProResDecoder.cpp
--- a/Engine/Plugins/Media/AppleProResDecoderElectra/Source/AppleProResDecoderElectra/Private/ProResDecoder/Mac/ElectraMediaProResDecoder.cpp
+++ b/Engine/Plugins/Media/AppleProResDecoderElectra/Source/AppleProResDecoderElectra/Private/ProResDecoder/Mac/ElectraMediaProResDecoder.cpp
@@ -511,6 +511,11 @@ private:
 class FProResVideoDecoderElectraFactory : public IElectraCodecFactory, public IElectraCodecModularFeature, public TSharedFromThis<FProResVideoDecoderElectraFactory, ESPMode::ThreadSafe>
 {
 public:
+	static constexpr uint32 Make4CC(const uint8 A, const uint8 B, const uint8 C, const uint8 D)
+	{
+		return (static_cast<uint32>(A) << 24) | (static_cast<uint32>(B) << 16) | (static_cast<uint32>(C) << 8) | static_cast<uint32>(D);
+	}
+
 	virtual ~FProResVideoDecoderElectraFactory()
 	{ }
 
@@ -519,9 +524,61 @@ public:
 		OutCodecFactories.Add(AsShared());
 	}
 
-	int32 SupportsFormat(const FString& InCodecFormat, bool bInEncoder, const TMap<FString, FVariant>& InOptions) const override
+	int32 SupportsFormat(TMap<FString, FVariant>& OutFormatInfo, const FString& InCodecFormat, bool bInEncoder, const TMap<FString, FVariant>& InOptions) const override
 	{
-		return bInEncoder || !Permitted4CCs.Contains(InCodecFormat) ? 0 : 1;
+		if (bInEncoder || !Permitted4CCs.Contains(InCodecFormat))
+		{
+			return 0;
+		}
+		switch((uint32)ElectraDecodersUtil::GetVariantValueSafeU64(InOptions, TEXT("codec_4cc"), 0))
+		{
+			case Make4CC('a','p','c','h'):
+			{
+				OutFormatInfo.Emplace(IElectraDecoderFormatInfo::HumanReadableFormatName, FVariant(FString(TEXT("Apple ProRes 422 High Quality"))));
+				break;
+			}
+			case Make4CC('a','p','c','n'):
+			{
+				OutFormatInfo.Emplace(IElectraDecoderFormatInfo::HumanReadableFormatName, FVariant(FString(TEXT("Apple ProRes 422 Standard Definition"))));
+				break;
+			}
+			case Make4CC('a','p','c','s'):
+			{
+				OutFormatInfo.Emplace(IElectraDecoderFormatInfo::HumanReadableFormatName, FVariant(FString(TEXT("Apple ProRes 422 LT"))));
+				break;
+			}
+			case Make4CC('a','p','c','o'):
+			{
+				OutFormatInfo.Emplace(IElectraDecoderFormatInfo::HumanReadableFormatName, FVariant(FString(TEXT("Apple ProRes 422 Proxy"))));
+				break;
+			}
+			case Make4CC('a','p','4','h'):
+			{
+				OutFormatInfo.Emplace(IElectraDecoderFormatInfo::HumanReadableFormatName, FVariant(FString(TEXT("Apple ProRes 4444"))));
+				break;
+			}
+			case Make4CC('a','p','4','x'):
+			{
+				OutFormatInfo.Emplace(IElectraDecoderFormatInfo::HumanReadableFormatName, FVariant(FString(TEXT("Apple ProRes 4444 XQ"))));
+				break;
+			}
+			case Make4CC('a','p','r','n'):
+			{
+				OutFormatInfo.Emplace(IElectraDecoderFormatInfo::HumanReadableFormatName, FVariant(FString(TEXT("Apple ProRes RAW"))));
+				break;
+			}
+			case Make4CC('a','p','r','h'):
+			{
+				OutFormatInfo.Emplace(IElectraDecoderFormatInfo::HumanReadableFormatName, FVariant(FString(TEXT("Apple ProRes RAW HQ"))));
+				break;
+			}
+			default:
+			{
+				break;
+			}
+		}
+		OutFormatInfo.Emplace(IElectraDecoderFormatInfo::IsEveryFrameKeyframe, FVariant(true));
+		return 1;
 	}
 
 	void GetConfigurationOptions(TMap<FString, FVariant>& OutOptions) const override
diff --git a/Engine/Plugins/Media/AppleProResDecoderElectra/Source/AppleProResDecoderElectra/Private/ProResDecoder/Windows/ElectraMediaProResDecoder.cpp b/Engine/Plugins/Media/AppleProResDecoderElectra/Source/AppleProResDecoderElectra/Private/ProResDecoder/Windows/ElectraMediaProResDecoder.cpp
--- a/Engine/Plugins/Media/AppleProResDecoderElectra/Source/AppleProResDecoderElectra/Private/ProResDecoder/Windows/ElectraMediaProResDecoder.cpp
+++ b/Engine/Plugins/Media/AppleProResDecoderElectra/Source/AppleProResDecoderElectra/Private/ProResDecoder/Windows/ElectraMediaProResDecoder.cpp
@@ -237,11 +237,6 @@ public:
 	{ }
 
 private:
-	static constexpr uint32 Make4CC(const uint8 A, const uint8 B, const uint8 C, const uint8 D)
-	{
-		return (static_cast<uint32>(A) << 24) | (static_cast<uint32>(B) << 16) | (static_cast<uint32>(C) << 8) | static_cast<uint32>(D);
-	}
-
 	struct FProResHeader
 	{
 		enum class EChromaType
@@ -401,6 +396,11 @@ private:
 class FProResVideoDecoderElectraFactory : public IElectraCodecFactory, public IElectraCodecModularFeature, public TSharedFromThis<FProResVideoDecoderElectraFactory, ESPMode::ThreadSafe>
 {
 public:
+	static constexpr uint32 Make4CC(const uint8 A, const uint8 B, const uint8 C, const uint8 D)
+	{
+		return (static_cast<uint32>(A) << 24) | (static_cast<uint32>(B) << 16) | (static_cast<uint32>(C) << 8) | static_cast<uint32>(D);
+	}
+
 	virtual ~FProResVideoDecoderElectraFactory()
 	{ }
 
@@ -409,9 +409,61 @@ public:
 		OutCodecFactories.Add(AsShared());
 	}
 
-	int32 SupportsFormat(const FString& InCodecFormat, bool bInEncoder, const TMap<FString, FVariant>& InOptions) const override
+	int32 SupportsFormat(TMap<FString, FVariant>& OutFormatInfo, const FString& InCodecFormat, bool bInEncoder, const TMap<FString, FVariant>& InOptions) const override
 	{
-		return bInEncoder || !Permitted4CCs.Contains(InCodecFormat) ? 0 : 1;
+		if (bInEncoder || !Permitted4CCs.Contains(InCodecFormat))
+		{
+			return 0;
+		}
+		switch((uint32)ElectraDecodersUtil::GetVariantValueSafeU64(InOptions, TEXT("codec_4cc"), 0))
+		{
+			case Make4CC('a','p','c','h'):
+			{
+				OutFormatInfo.Emplace(IElectraDecoderFormatInfo::HumanReadableFormatName, FVariant(FString(TEXT("Apple ProRes 422 High Quality"))));
+				break;
+			}
+			case Make4CC('a','p','c','n'):
+			{
+				OutFormatInfo.Emplace(IElectraDecoderFormatInfo::HumanReadableFormatName, FVariant(FString(TEXT("Apple ProRes 422 Standard Definition"))));
+				break;
+			}
+			case Make4CC('a','p','c','s'):
+			{
+				OutFormatInfo.Emplace(IElectraDecoderFormatInfo::HumanReadableFormatName, FVariant(FString(TEXT("Apple ProRes 422 LT"))));
+				break;
+			}
+			case Make4CC('a','p','c','o'):
+			{
+				OutFormatInfo.Emplace(IElectraDecoderFormatInfo::HumanReadableFormatName, FVariant(FString(TEXT("Apple ProRes 422 Proxy"))));
+				break;
+			}
+			case Make4CC('a','p','4','h'):
+			{
+				OutFormatInfo.Emplace(IElectraDecoderFormatInfo::HumanReadableFormatName, FVariant(FString(TEXT("Apple ProRes 4444"))));
+				break;
+			}
+			case Make4CC('a','p','4','x'):
+			{
+				OutFormatInfo.Emplace(IElectraDecoderFormatInfo::HumanReadableFormatName, FVariant(FString(TEXT("Apple ProRes 4444 XQ"))));
+				break;
+			}
+			case Make4CC('a','p','r','n'):
+			{
+				OutFormatInfo.Emplace(IElectraDecoderFormatInfo::HumanReadableFormatName, FVariant(FString(TEXT("Apple ProRes RAW"))));
+				break;
+			}
+			case Make4CC('a','p','r','h'):
+			{
+				OutFormatInfo.Emplace(IElectraDecoderFormatInfo::HumanReadableFormatName, FVariant(FString(TEXT("Apple ProRes RAW HQ"))));
+				break;
+			}
+			default:
+			{
+				break;
+			}
+		}
+		OutFormatInfo.Emplace(IElectraDecoderFormatInfo::IsEveryFrameKeyframe, FVariant(true));
+		return 1;
 	}
 
 	void GetConfigurationOptions(TMap<FString, FVariant>& OutOptions) const override
diff --git a/Engine/Plugins/Media/AvidDNxHDDecoderElectra/Source/Private/Windows/ElectraMediaAvidDNxHDDecoder.cpp b/Engine/Plugins/Media/AvidDNxHDDecoderElectra/Source/Private/Windows/ElectraMediaAvidDNxHDDecoder.cpp
--- a/Engine/Plugins/Media/AvidDNxHDDecoderElectra/Source/Private/Windows/ElectraMediaAvidDNxHDDecoder.cpp
+++ b/Engine/Plugins/Media/AvidDNxHDDecoderElectra/Source/Private/Windows/ElectraMediaAvidDNxHDDecoder.cpp
@@ -258,11 +258,6 @@ public:
 	{ }
 
 private:
-	static constexpr uint32 Make4CC(const uint8 A, const uint8 B, const uint8 C, const uint8 D)
-	{
-		return (static_cast<uint32>(A) << 24) | (static_cast<uint32>(B) << 16) | (static_cast<uint32>(C) << 8) | static_cast<uint32>(D);
-	}
-
 	bool PostError(int32 ApiReturnValue, FString Message, int32 Code);
 
 	int32 DisplayWidth = 0;
@@ -315,6 +310,11 @@ private:
 class FAvidDNxHDVideoDecoderElectraFactory : public IElectraCodecFactory, public IElectraCodecModularFeature, public TSharedFromThis<FAvidDNxHDVideoDecoderElectraFactory, ESPMode::ThreadSafe>
 {
 public:
+	static constexpr uint32 Make4CC(const uint8 A, const uint8 B, const uint8 C, const uint8 D)
+	{
+		return (static_cast<uint32>(A) << 24) | (static_cast<uint32>(B) << 16) | (static_cast<uint32>(C) << 8) | static_cast<uint32>(D);
+	}
+
 	virtual ~FAvidDNxHDVideoDecoderElectraFactory()
 	{ }
 
@@ -323,9 +323,26 @@ public:
 		OutCodecFactories.Add(AsShared());
 	}
 
-	int32 SupportsFormat(const FString& InCodecFormat, bool bInEncoder, const TMap<FString, FVariant>& InOptions) const override
+	int32 SupportsFormat(TMap<FString, FVariant>& OutFormatInfo, const FString& InCodecFormat, bool bInEncoder, const TMap<FString, FVariant>& InOptions) const override
 	{
-		return bInEncoder || !Permitted4CCs.Contains(InCodecFormat) ? 0 : 1;
+		if (bInEncoder || !Permitted4CCs.Contains(InCodecFormat))
+		{
+			return 0;
+		}
+		switch((uint32)ElectraDecodersUtil::GetVariantValueSafeU64(InOptions, TEXT("codec_4cc"), 0))
+		{
+			case Make4CC('A','V','d','h'):
+			{
+				OutFormatInfo.Emplace(IElectraDecoderFormatInfo::HumanReadableFormatName, FVariant(FString(TEXT("Avid DNxHD"))));
+				break;
+			}
+			default:
+			{
+				break;
+			}
+		}
+		OutFormatInfo.Emplace(IElectraDecoderFormatInfo::IsEveryFrameKeyframe, FVariant(true));
+		return 1;
 	}
 
 	void GetConfigurationOptions(TMap<FString, FVariant>& OutOptions) const override
diff --git a/Engine/Plugins/Media/D3D12VideoDecodersElectra/Source/Private/VideoDecoder_D3D12.cpp b/Engine/Plugins/Media/D3D12VideoDecodersElectra/Source/Private/VideoDecoder_D3D12.cpp
--- a/Engine/Plugins/Media/D3D12VideoDecodersElectra/Source/Private/VideoDecoder_D3D12.cpp
+++ b/Engine/Plugins/Media/D3D12VideoDecodersElectra/Source/Private/VideoDecoder_D3D12.cpp
@@ -439,7 +439,7 @@ public:
 		return Codec;
 	}
 
-	int32 SupportsFormat(const FString& InCodecFormat, bool bInEncoder, const TMap<FString, FVariant>& InOptions) const override
+	int32 SupportsFormat(TMap<FString, FVariant>& OutFormatInfo, const FString& InCodecFormat, bool bInEncoder, const TMap<FString, FVariant>& InOptions) const override
 	{
 		if (bDoNotUseThisDecoder)
 		{
diff --git a/Engine/Plugins/Media/ElectraCodecs/Source/ElectraCodecFactory/Private/ElectraCodecFactoryModule.cpp b/Engine/Plugins/Media/ElectraCodecs/Source/ElectraCodecFactory/Private/ElectraCodecFactoryModule.cpp
--- a/Engine/Plugins/Media/ElectraCodecs/Source/ElectraCodecFactory/Private/ElectraCodecFactoryModule.cpp
+++ b/Engine/Plugins/Media/ElectraCodecs/Source/ElectraCodecFactory/Private/ElectraCodecFactoryModule.cpp
@@ -68,13 +68,13 @@ public:
 	{
 		IModularFeatures::Get().UnregisterModularFeature(GetModularFeatureName(), &StandardDecoders);
 	}
-	
+
 	bool SupportsDynamicReloading() override
 	{
 		return false;
 	}
 
-	TSharedPtr<IElectraCodecFactory, ESPMode::ThreadSafe> GetBestFactoryForFormat(const FString& InCodecFormat, bool bInEncoder, const TMap<FString, FVariant>& InOptions) override;
+	TSharedPtr<IElectraCodecFactory, ESPMode::ThreadSafe> GetBestFactoryForFormat(TMap<FString, FVariant>& OutFormatInfo, const FString& InCodecFormat, bool bInEncoder, const TMap<FString, FVariant>& InOptions) override;
 
 private:
 	FElectraCodecRegistry StandardDecoders;
@@ -83,7 +83,7 @@ private:
 IMPLEMENT_MODULE(FElectraCodecFactoryModule, ElectraCodecFactory);
 
 
-TSharedPtr<IElectraCodecFactory, ESPMode::ThreadSafe> FElectraCodecFactoryModule::GetBestFactoryForFormat(const FString& InCodecFormat, bool bInEncoder, const TMap<FString, FVariant>& InOptions)
+TSharedPtr<IElectraCodecFactory, ESPMode::ThreadSafe> FElectraCodecFactoryModule::GetBestFactoryForFormat(TMap<FString, FVariant>& OutFormatInfo, const FString& InCodecFormat, bool bInEncoder, const TMap<FString, FVariant>& InOptions)
 {
 	IModularFeatures::FScopedLockModularFeatureList ScopedLockModularFeatureList;
 	// Get the list of registered factories
@@ -97,11 +97,13 @@ TSharedPtr<IElectraCodecFactory, ESPMode::ThreadSafe> FElectraCodecFactoryModule
 		for(int32 j=0; j<FormatFactories.Num(); ++j)
 		{
 			TSharedPtr<IElectraCodecFactory, ESPMode::ThreadSafe> FormatFactory = FormatFactories[j].Pin();
-			int32 SupportedPriority = FormatFactory.IsValid() ? FormatFactory->SupportsFormat(InCodecFormat, bInEncoder, InOptions) : 0;
+			TMap<FString, FVariant> FormatInfo;
+			int32 SupportedPriority = FormatFactory.IsValid() ? FormatFactory->SupportsFormat(FormatInfo, InCodecFormat, bInEncoder, InOptions) : 0;
 			if (SupportedPriority > BestPriority)
 			{
 				BestPriority = SupportedPriority;
 				BestFactory = FormatFactory;
+				OutFormatInfo = MoveTemp(FormatInfo);
 			}
 		}
 	}
diff --git a/Engine/Plugins/Media/ElectraCodecs/Source/ElectraCodecFactory/Public/IElectraCodecFactory.h b/Engine/Plugins/Media/ElectraCodecs/Source/ElectraCodecFactory/Public/IElectraCodecFactory.h
--- a/Engine/Plugins/Media/ElectraCodecs/Source/ElectraCodecFactory/Public/IElectraCodecFactory.h
+++ b/Engine/Plugins/Media/ElectraCodecs/Source/ElectraCodecFactory/Public/IElectraCodecFactory.h
@@ -19,27 +19,27 @@ public:
 	/**
 	 * Queries whether or not this codec factory can create a decoder or encoder of the specified format.
 	 * The codec format string is the "codecs" part (RFC 6381) of a MIME type like "video/mp4" (RFC 6838).
-	 * 
+	 *
 	 * For example (back quote characters `` are for clarity only, they do not appear in actual strings):
 	 *   `avc1.4d002a`  if the full MIME type were `video/mp4; codecs="avc1.4d002a"`
 	 *   `mp4a.40.2`    if the full MIME type were `audio/mp4; codecs="mp4a.40.2"`
-	 * 
+	 *
 	 * Note that only a single format must be specified. Should there be several codecs specified in a
 	 * MIME type they must be separated and each format queried for individually.
-	 * 
+	 *
 	 * Format profiles may be omitted if they are not known. Asking for an `avc1` decoder might result
 	 * in a supported decoder, but there is no guarantee that it is capable of decoding every profile.
-	 * 
+	 *
 	 * Format names are case sensitive. Profile case depends on the format.
-	 * 
+	 *
 	 * Additional options may be provided to specify the format more closely, which may help selecting
 	 * the best suited implementation. See below.
-	 * 
+	 *
 	 * Returns 0 if not supported.
 	 * If supported the return value indicates a priority value. If multiple factories are registered
 	 * that claim support for the format the one with the highest priority is chosen.
 	 */
-	virtual int32 SupportsFormat(const FString& InCodecFormat, bool bInEncoder, const TMap<FString, FVariant>& InOptions) const = 0;
+	virtual int32 SupportsFormat(TMap<FString, FVariant>& OutFormatInfo, const FString& InCodecFormat, bool bInEncoder, const TMap<FString, FVariant>& InOptions) const = 0;
 
 
 	/**
@@ -72,7 +72,7 @@ public:
 
 		"csd" (ByteArray)
 			- AAC audio: codec specific data
-		
+
 		"dcr" (ByteArray)
 			- AAC audio: decoder configuration record
 
@@ -81,11 +81,11 @@ public:
 
 		"$enda_box" (ByteArray)
 		    - The raw contents of the mp4 'enda' box if it exists without the first 8 bytes (size and atom)
-		
+
 		"$FormatSpecificFlags" (int64)
 			- The "FormatSpecificFlags" from a version 2 AudioSampleEntry
 
-	
+
 	Video decoder:
 
 		"width" (uint32)
@@ -93,7 +93,7 @@ public:
 
 		"height" (uint32)
 			Height in pixels of the decoded image (active pixels after cropping). May be 0 if unknown.
-		
+
 		"bps" (int64)
 			Bits per second of the video stream. May be <=0 if unknown.
 
diff --git a/Engine/Plugins/Media/ElectraCodecs/Source/ElectraCodecFactory/Public/IElectraCodecFactoryModule.h b/Engine/Plugins/Media/ElectraCodecs/Source/ElectraCodecFactory/Public/IElectraCodecFactoryModule.h
--- a/Engine/Plugins/Media/ElectraCodecs/Source/ElectraCodecFactory/Public/IElectraCodecFactoryModule.h
+++ b/Engine/Plugins/Media/ElectraCodecs/Source/ElectraCodecFactory/Public/IElectraCodecFactoryModule.h
@@ -27,5 +27,5 @@ public:
 		return FeatureName;
 	}
 
-	virtual TSharedPtr<IElectraCodecFactory, ESPMode::ThreadSafe> GetBestFactoryForFormat(const FString& InCodecFormat, bool bInEncoder, const TMap<FString, FVariant>& InOptions) = 0;
+	virtual TSharedPtr<IElectraCodecFactory, ESPMode::ThreadSafe> GetBestFactoryForFormat(TMap<FString, FVariant>& OutFormatInfo, const FString& InCodecFormat, bool bInEncoder, const TMap<FString, FVariant>& InOptions) = 0;
 };
diff --git a/Engine/Plugins/Media/ElectraCodecs/Source/ElectraDecoders/Private/Android/aac/AAC_AudioDecoder_Android.cpp b/Engine/Plugins/Media/ElectraCodecs/Source/ElectraDecoders/Private/Android/aac/AAC_AudioDecoder_Android.cpp
--- a/Engine/Plugins/Media/ElectraCodecs/Source/ElectraDecoders/Private/Android/aac/AAC_AudioDecoder_Android.cpp
+++ b/Engine/Plugins/Media/ElectraCodecs/Source/ElectraDecoders/Private/Android/aac/AAC_AudioDecoder_Android.cpp
@@ -12,7 +12,7 @@ public:
 	virtual ~FAACAudioDecoderFactoryAndroid()
 	{}
 
-	int32 SupportsFormat(const FString& InCodecFormat, bool bInEncoder, const TMap<FString, FVariant>& InOptions) const override
+	int32 SupportsFormat(TMap<FString, FVariant>& OutFormatInfo, const FString& InCodecFormat, bool bInEncoder, const TMap<FString, FVariant>& InOptions) const override
 	{
 		// Encoder? Not supported here!
 		if (bInEncoder)
@@ -63,10 +63,9 @@ public:
 	{
 		IElectraAudioDecoderAAC_Android::GetConfigurationOptions(OutOptions);
 	}
-	
+
 	TSharedPtr<IElectraDecoder, ESPMode::ThreadSafe> CreateDecoderForFormat(const FString& InCodecFormat, const TMap<FString, FVariant>& InOptions, TSharedPtr<IElectraDecoderResourceDelegate, ESPMode::ThreadSafe> InResourceDelegate) override
 	{
-		check(SupportsFormat(InCodecFormat, false, InOptions));
 		return IElectraAudioDecoderAAC_Android::Create(InOptions, InResourceDelegate);
 	}
 
diff --git a/Engine/Plugins/Media/ElectraCodecs/Source/ElectraDecoders/Private/Android/h264/H264_VideoDecoder_Android.cpp b/Engine/Plugins/Media/ElectraCodecs/Source/ElectraDecoders/Private/Android/h264/H264_VideoDecoder_Android.cpp
--- a/Engine/Plugins/Media/ElectraCodecs/Source/ElectraDecoders/Private/Android/h264/H264_VideoDecoder_Android.cpp
+++ b/Engine/Plugins/Media/ElectraCodecs/Source/ElectraDecoders/Private/Android/h264/H264_VideoDecoder_Android.cpp
@@ -12,7 +12,7 @@ public:
 	virtual ~FH264VideoDecoderFactoryAndroid()
 	{}
 
-	int32 SupportsFormat(const FString& InCodecFormat, bool bInEncoder, const TMap<FString, FVariant>& InOptions) const override
+	int32 SupportsFormat(TMap<FString, FVariant>& OutFormatInfo, const FString& InCodecFormat, bool bInEncoder, const TMap<FString, FVariant>& InOptions) const override
 	{
 		// Encoder? Not supported here!
 		if (bInEncoder)
@@ -65,7 +65,6 @@ public:
 
 	TSharedPtr<IElectraDecoder, ESPMode::ThreadSafe> CreateDecoderForFormat(const FString& InCodecFormat, const TMap<FString, FVariant>& InOptions, TSharedPtr<IElectraDecoderResourceDelegate, ESPMode::ThreadSafe> InResourceDelegate) override
 	{
-		check(SupportsFormat(InCodecFormat, false, InOptions));
 		return IElectraVideoDecoderH264_Android::Create(InOptions, InResourceDelegate);
 	}
 
diff --git a/Engine/Plugins/Media/ElectraCodecs/Source/ElectraDecoders/Private/Android/h265/H265_VideoDecoder_Android.cpp b/Engine/Plugins/Media/ElectraCodecs/Source/ElectraDecoders/Private/Android/h265/H265_VideoDecoder_Android.cpp
--- a/Engine/Plugins/Media/ElectraCodecs/Source/ElectraDecoders/Private/Android/h265/H265_VideoDecoder_Android.cpp
+++ b/Engine/Plugins/Media/ElectraCodecs/Source/ElectraDecoders/Private/Android/h265/H265_VideoDecoder_Android.cpp
@@ -12,7 +12,7 @@ public:
 	virtual ~FH265VideoDecoderFactoryAndroid()
 	{}
 
-	int32 SupportsFormat(const FString& InCodecFormat, bool bInEncoder, const TMap<FString, FVariant>& InOptions) const override
+	int32 SupportsFormat(TMap<FString, FVariant>& OutFormatInfo, const FString& InCodecFormat, bool bInEncoder, const TMap<FString, FVariant>& InOptions) const override
 	{
 		// Encoder? Not supported here!
 		if (bInEncoder)
@@ -83,7 +83,6 @@ public:
 
 	TSharedPtr<IElectraDecoder, ESPMode::ThreadSafe> CreateDecoderForFormat(const FString& InCodecFormat, const TMap<FString, FVariant>& InOptions, TSharedPtr<IElectraDecoderResourceDelegate, ESPMode::ThreadSafe> InResourceDelegate) override
 	{
-		check(SupportsFormat(InCodecFormat, false, InOptions));
 		return IElectraVideoDecoderH265_Android::Create(InOptions, InResourceDelegate);
 	}
 
diff --git a/Engine/Plugins/Media/ElectraCodecs/Source/ElectraDecoders/Private/Android/vpx/VPx_VideoDecoder_Android.cpp b/Engine/Plugins/Media/ElectraCodecs/Source/ElectraDecoders/Private/Android/vpx/VPx_VideoDecoder_Android.cpp
--- a/Engine/Plugins/Media/ElectraCodecs/Source/ElectraDecoders/Private/Android/vpx/VPx_VideoDecoder_Android.cpp
+++ b/Engine/Plugins/Media/ElectraCodecs/Source/ElectraDecoders/Private/Android/vpx/VPx_VideoDecoder_Android.cpp
@@ -12,7 +12,7 @@ public:
 	virtual ~FVPxVideoDecoderFactoryAndroid()
 	{}
 
-	int32 SupportsFormat(const FString& InCodecFormat, bool bInEncoder, const TMap<FString, FVariant>& InOptions) const override
+	int32 SupportsFormat(TMap<FString, FVariant>& OutFormatInfo, const FString& InCodecFormat, bool bInEncoder, const TMap<FString, FVariant>& InOptions) const override
 	{
 		// Encoder? Not supported here!
 		if (bInEncoder)
@@ -43,7 +43,6 @@ public:
 
 	TSharedPtr<IElectraDecoder, ESPMode::ThreadSafe> CreateDecoderForFormat(const FString& InCodecFormat, const TMap<FString, FVariant>& InOptions, TSharedPtr<IElectraDecoderResourceDelegate, ESPMode::ThreadSafe> InResourceDelegate) override
 	{
-		check(SupportsFormat(InCodecFormat, false, InOptions));
 		return IElectraVideoDecoderVPx_Android::Create(InOptions, InResourceDelegate);
 	}
 
diff --git a/Engine/Plugins/Media/ElectraCodecs/Source/ElectraDecoders/Private/Apple/aac/AAC_AudioDecoder_Apple.cpp b/Engine/Plugins/Media/ElectraCodecs/Source/ElectraDecoders/Private/Apple/aac/AAC_AudioDecoder_Apple.cpp
--- a/Engine/Plugins/Media/ElectraCodecs/Source/ElectraDecoders/Private/Apple/aac/AAC_AudioDecoder_Apple.cpp
+++ b/Engine/Plugins/Media/ElectraCodecs/Source/ElectraDecoders/Private/Apple/aac/AAC_AudioDecoder_Apple.cpp
@@ -14,7 +14,7 @@ public:
 	virtual ~FAACAudioDecoderFactoryApple()
 	{}
 
-	int32 SupportsFormat(const FString& InCodecFormat, bool bInEncoder, const TMap<FString, FVariant>& InOptions) const override
+	int32 SupportsFormat(TMap<FString, FVariant>& OutFormatInfo, const FString& InCodecFormat, bool bInEncoder, const TMap<FString, FVariant>& InOptions) const override
 	{
 		// Encoder? Not supported here!
 		if (bInEncoder)
@@ -72,7 +72,6 @@ public:
 
 	TSharedPtr<IElectraDecoder, ESPMode::ThreadSafe> CreateDecoderForFormat(const FString& InCodecFormat, const TMap<FString, FVariant>& InOptions, TSharedPtr<IElectraDecoderResourceDelegate, ESPMode::ThreadSafe> InResourceDelegate) override
 	{
-		check(SupportsFormat(InCodecFormat, false, InOptions));
 		return IElectraAudioDecoderAAC_Apple::Create(InOptions, InResourceDelegate);
 	}
 
diff --git a/Engine/Plugins/Media/ElectraCodecs/Source/ElectraDecoders/Private/Apple/h264/H264_VideoDecoder_Apple.cpp b/Engine/Plugins/Media/ElectraCodecs/Source/ElectraDecoders/Private/Apple/h264/H264_VideoDecoder_Apple.cpp
--- a/Engine/Plugins/Media/ElectraCodecs/Source/ElectraDecoders/Private/Apple/h264/H264_VideoDecoder_Apple.cpp
+++ b/Engine/Plugins/Media/ElectraCodecs/Source/ElectraDecoders/Private/Apple/h264/H264_VideoDecoder_Apple.cpp
@@ -14,7 +14,7 @@ public:
 	virtual ~FH264VideoDecoderFactoryApple()
 	{}
 
-	int32 SupportsFormat(const FString& InCodecFormat, bool bInEncoder, const TMap<FString, FVariant>& InOptions) const override
+	int32 SupportsFormat(TMap<FString, FVariant>& OutFormatInfo, const FString& InCodecFormat, bool bInEncoder, const TMap<FString, FVariant>& InOptions) const override
 	{
 		// Encoder? Not supported here!
 		if (bInEncoder)
@@ -85,7 +85,6 @@ public:
 
 	TSharedPtr<IElectraDecoder, ESPMode::ThreadSafe> CreateDecoderForFormat(const FString& InCodecFormat, const TMap<FString, FVariant>& InOptions, TSharedPtr<IElectraDecoderResourceDelegate, ESPMode::ThreadSafe> InResourceDelegate) override
 	{
-		check(SupportsFormat(InCodecFormat, false, InOptions));
 		return IElectraVideoDecoderH264_Apple::Create(InOptions, InResourceDelegate);
 	}
 
diff --git a/Engine/Plugins/Media/ElectraCodecs/Source/ElectraDecoders/Private/Apple/h265/H265_VideoDecoder_Apple.cpp b/Engine/Plugins/Media/ElectraCodecs/Source/ElectraDecoders/Private/Apple/h265/H265_VideoDecoder_Apple.cpp
--- a/Engine/Plugins/Media/ElectraCodecs/Source/ElectraDecoders/Private/Apple/h265/H265_VideoDecoder_Apple.cpp
+++ b/Engine/Plugins/Media/ElectraCodecs/Source/ElectraDecoders/Private/Apple/h265/H265_VideoDecoder_Apple.cpp
@@ -14,7 +14,7 @@ public:
 	virtual ~FH265VideoDecoderFactoryApple()
 	{}
 
-	int32 SupportsFormat(const FString& InCodecFormat, bool bInEncoder, const TMap<FString, FVariant>& InOptions) const override
+	int32 SupportsFormat(TMap<FString, FVariant>& OutFormatInfo, const FString& InCodecFormat, bool bInEncoder, const TMap<FString, FVariant>& InOptions) const override
 	{
 		// Encoder? Not supported here!
 		if (bInEncoder)
@@ -85,7 +85,6 @@ public:
 
 	TSharedPtr<IElectraDecoder, ESPMode::ThreadSafe> CreateDecoderForFormat(const FString& InCodecFormat, const TMap<FString, FVariant>& InOptions, TSharedPtr<IElectraDecoderResourceDelegate, ESPMode::ThreadSafe> InResourceDelegate) override
 	{
-		check(SupportsFormat(InCodecFormat, false, InOptions));
 		return IElectraVideoDecoderH265_Apple::Create(InOptions, InResourceDelegate);
 	}
 
diff --git a/Engine/Plugins/Media/ElectraCodecs/Source/ElectraDecoders/Private/ElectraMediaMP3Decoder.cpp b/Engine/Plugins/Media/ElectraCodecs/Source/ElectraDecoders/Private/ElectraMediaMP3Decoder.cpp
--- a/Engine/Plugins/Media/ElectraCodecs/Source/ElectraDecoders/Private/ElectraMediaMP3Decoder.cpp
+++ b/Engine/Plugins/Media/ElectraCodecs/Source/ElectraDecoders/Private/ElectraMediaMP3Decoder.cpp
@@ -229,7 +229,7 @@ public:
 		OutCodecFactories.Add(AsShared());
 	}
 
-	int32 SupportsFormat(const FString& InCodecFormat, bool bInEncoder, const TMap<FString, FVariant>& InOptions) const override
+	int32 SupportsFormat(TMap<FString, FVariant>& OutFormatInfo, const FString& InCodecFormat, bool bInEncoder, const TMap<FString, FVariant>& InOptions) const override
 	{
 		// Quick check if this is an ask for an encoder or for a 4CC we do not support.
 		if (bInEncoder || !Permitted4CCs.Contains(InCodecFormat))
diff --git a/Engine/Plugins/Media/ElectraCodecs/Source/ElectraDecoders/Private/Linux/aac/AAC_AudioDecoder_Linux.cpp b/Engine/Plugins/Media/ElectraCodecs/Source/ElectraDecoders/Private/Linux/aac/AAC_AudioDecoder_Linux.cpp
--- a/Engine/Plugins/Media/ElectraCodecs/Source/ElectraDecoders/Private/Linux/aac/AAC_AudioDecoder_Linux.cpp
+++ b/Engine/Plugins/Media/ElectraCodecs/Source/ElectraDecoders/Private/Linux/aac/AAC_AudioDecoder_Linux.cpp
@@ -18,7 +18,7 @@ public:
 	virtual ~FAACAudioDecoderFactoryLinux()
 	{}
 
-	int32 SupportsFormat(const FString& InCodecFormat, bool bInEncoder, const TMap<FString, FVariant>& InOptions) const override
+	int32 SupportsFormat(TMap<FString, FVariant>& OutFormatInfo, const FString& InCodecFormat, bool bInEncoder, const TMap<FString, FVariant>& InOptions) const override
 	{
 		// Encoder? Not supported here!
 		if (bInEncoder)
@@ -81,7 +81,6 @@ public:
 
 	TSharedPtr<IElectraDecoder, ESPMode::ThreadSafe> CreateDecoderForFormat(const FString& InCodecFormat, const TMap<FString, FVariant>& InOptions, TSharedPtr<IElectraDecoderResourceDelegate, ESPMode::ThreadSafe> InResourceDelegate) override
 	{
-		check(SupportsFormat(InCodecFormat, false, InOptions));
 		return IElectraAudioDecoderAAC_Linux::Create(InOptions, InResourceDelegate);
 	}
 
diff --git a/Engine/Plugins/Media/ElectraCodecs/Source/ElectraDecoders/Private/Linux/h264/H264_VideoDecoder_Linux.cpp b/Engine/Plugins/Media/ElectraCodecs/Source/ElectraDecoders/Private/Linux/h264/H264_VideoDecoder_Linux.cpp
--- a/Engine/Plugins/Media/ElectraCodecs/Source/ElectraDecoders/Private/Linux/h264/H264_VideoDecoder_Linux.cpp
+++ b/Engine/Plugins/Media/ElectraCodecs/Source/ElectraDecoders/Private/Linux/h264/H264_VideoDecoder_Linux.cpp
@@ -16,7 +16,7 @@ public:
 	virtual ~FH264VideoDecoderFactoryLinux()
 	{}
 
-	int32 SupportsFormat(const FString& InCodecFormat, bool bInEncoder, const TMap<FString, FVariant>& InOptions) const override
+	int32 SupportsFormat(TMap<FString, FVariant>& OutFormatInfo, const FString& InCodecFormat, bool bInEncoder, const TMap<FString, FVariant>& InOptions) const override
 	{
 		// Encoder? Not supported here!
 		if (bInEncoder)
@@ -93,7 +93,6 @@ public:
 
 	TSharedPtr<IElectraDecoder, ESPMode::ThreadSafe> CreateDecoderForFormat(const FString& InCodecFormat, const TMap<FString, FVariant>& InOptions, TSharedPtr<IElectraDecoderResourceDelegate, ESPMode::ThreadSafe> InResourceDelegate) override
 	{
-		check(SupportsFormat(InCodecFormat, false, InOptions));
 		return IElectraVideoDecoderH264_Linux::Create(InOptions, InResourceDelegate);
 	}
 
diff --git a/Engine/Plugins/Media/ElectraCodecs/Source/ElectraDecoders/Private/Linux/h265/H265_VideoDecoder_Linux.cpp b/Engine/Plugins/Media/ElectraCodecs/Source/ElectraDecoders/Private/Linux/h265/H265_VideoDecoder_Linux.cpp
--- a/Engine/Plugins/Media/ElectraCodecs/Source/ElectraDecoders/Private/Linux/h265/H265_VideoDecoder_Linux.cpp
+++ b/Engine/Plugins/Media/ElectraCodecs/Source/ElectraDecoders/Private/Linux/h265/H265_VideoDecoder_Linux.cpp
@@ -16,7 +16,7 @@ public:
 	virtual ~FH265VideoDecoderFactoryLinux()
 	{}
 
-	int32 SupportsFormat(const FString& InCodecFormat, bool bInEncoder, const TMap<FString, FVariant>& InOptions) const override
+	int32 SupportsFormat(TMap<FString, FVariant>& OutFormatInfo, const FString& InCodecFormat, bool bInEncoder, const TMap<FString, FVariant>& InOptions) const override
 	{
 		// Encoder? Not supported here!
 		if (bInEncoder)
@@ -93,7 +93,6 @@ public:
 
 	TSharedPtr<IElectraDecoder, ESPMode::ThreadSafe> CreateDecoderForFormat(const FString& InCodecFormat, const TMap<FString, FVariant>& InOptions, TSharedPtr<IElectraDecoderResourceDelegate, ESPMode::ThreadSafe> InResourceDelegate) override
 	{
-		check(SupportsFormat(InCodecFormat, false, InOptions));
 		return IElectraVideoDecoderH265_Linux::Create(InOptions, InResourceDelegate);
 	}
 
diff --git a/Engine/Plugins/Media/ElectraCodecs/Source/ElectraDecoders/Private/Linux/vpx/VPx_VideoDecoder_Linux.cpp b/Engine/Plugins/Media/ElectraCodecs/Source/ElectraDecoders/Private/Linux/vpx/VPx_VideoDecoder_Linux.cpp
--- a/Engine/Plugins/Media/ElectraCodecs/Source/ElectraDecoders/Private/Linux/vpx/VPx_VideoDecoder_Linux.cpp
+++ b/Engine/Plugins/Media/ElectraCodecs/Source/ElectraDecoders/Private/Linux/vpx/VPx_VideoDecoder_Linux.cpp
@@ -17,7 +17,7 @@ public:
 	virtual ~FVPxVideoDecoderFactoryLinux()
 	{}
 
-	int32 SupportsFormat(const FString& InCodecFormat, bool bInEncoder, const TMap<FString, FVariant>& InOptions) const override
+	int32 SupportsFormat(TMap<FString, FVariant>& OutFormatInfo, const FString& InCodecFormat, bool bInEncoder, const TMap<FString, FVariant>& InOptions) const override
 	{
 		// Encoder? Not supported here!
 		if (bInEncoder)
diff --git a/Engine/Plugins/Media/ElectraCodecs/Source/ElectraDecoders/Private/Windows/aac/AAC_AudioDecoder_Windows.cpp b/Engine/Plugins/Media/ElectraCodecs/Source/ElectraDecoders/Private/Windows/aac/AAC_AudioDecoder_Windows.cpp
--- a/Engine/Plugins/Media/ElectraCodecs/Source/ElectraDecoders/Private/Windows/aac/AAC_AudioDecoder_Windows.cpp
+++ b/Engine/Plugins/Media/ElectraCodecs/Source/ElectraDecoders/Private/Windows/aac/AAC_AudioDecoder_Windows.cpp
@@ -14,7 +14,7 @@ public:
 	virtual ~FAACAudioDecoderFactoryWindows()
 	{}
 
-	int32 SupportsFormat(const FString& InCodecFormat, bool bInEncoder, const TMap<FString, FVariant>& InOptions) const override
+	int32 SupportsFormat(TMap<FString, FVariant>& OutFormatInfo, const FString& InCodecFormat, bool bInEncoder, const TMap<FString, FVariant>& InOptions) const override
 	{
 		// Encoder? Not supported here!
 		if (bInEncoder)
@@ -60,7 +60,7 @@ public:
 
 		return 1;
 	}
-	
+
 	void GetConfigurationOptions(TMap<FString, FVariant>& OutOptions) const override
 	{
 		IElectraAudioDecoderAAC_DX::GetConfigurationOptions(OutOptions);
@@ -68,7 +68,6 @@ public:
 
 	TSharedPtr<IElectraDecoder, ESPMode::ThreadSafe> CreateDecoderForFormat(const FString& InCodecFormat, const TMap<FString, FVariant>& InOptions, TSharedPtr<IElectraDecoderResourceDelegate, ESPMode::ThreadSafe> InResourceDelegate) override
 	{
-		check(SupportsFormat(InCodecFormat, false, InOptions));
 		return IElectraAudioDecoderAAC_DX::Create(InOptions, InResourceDelegate);
 	}
 
diff --git a/Engine/Plugins/Media/ElectraCodecs/Source/ElectraDecoders/Private/Windows/h264/H264_VideoDecoder_Windows.cpp b/Engine/Plugins/Media/ElectraCodecs/Source/ElectraDecoders/Private/Windows/h264/H264_VideoDecoder_Windows.cpp
--- a/Engine/Plugins/Media/ElectraCodecs/Source/ElectraDecoders/Private/Windows/h264/H264_VideoDecoder_Windows.cpp
+++ b/Engine/Plugins/Media/ElectraCodecs/Source/ElectraDecoders/Private/Windows/h264/H264_VideoDecoder_Windows.cpp
@@ -46,7 +46,7 @@ public:
 	virtual ~FH264VideoDecoderFactoryWindows()
 	{}
 
-	int32 SupportsFormat(const FString& InCodecFormat, bool bInEncoder, const TMap<FString, FVariant>& InOptions) const override
+	int32 SupportsFormat(TMap<FString, FVariant>& OutFormatInfo, const FString& InCodecFormat, bool bInEncoder, const TMap<FString, FVariant>& InOptions) const override
 	{
 		// Encoder? Not supported here!
 		if (bInEncoder)
@@ -117,7 +117,6 @@ public:
 
 	TSharedPtr<IElectraDecoder, ESPMode::ThreadSafe> CreateDecoderForFormat(const FString& InCodecFormat, const TMap<FString, FVariant>& InOptions, TSharedPtr<IElectraDecoderResourceDelegate, ESPMode::ThreadSafe> InResourceDelegate) override
 	{
-		check(SupportsFormat(InCodecFormat, false, InOptions));
 		return IElectraVideoDecoderH264_DX::Create(InOptions, InResourceDelegate);
 	}
 
diff --git a/Engine/Plugins/Media/ElectraCodecs/Source/ElectraDecoders/Private/Windows/h265/H265_VideoDecoder_Windows.cpp b/Engine/Plugins/Media/ElectraCodecs/Source/ElectraDecoders/Private/Windows/h265/H265_VideoDecoder_Windows.cpp
--- a/Engine/Plugins/Media/ElectraCodecs/Source/ElectraDecoders/Private/Windows/h265/H265_VideoDecoder_Windows.cpp
+++ b/Engine/Plugins/Media/ElectraCodecs/Source/ElectraDecoders/Private/Windows/h265/H265_VideoDecoder_Windows.cpp
@@ -33,7 +33,7 @@ public:
 	virtual ~FH265VideoDecoderFactoryWindows()
 	{}
 
-	int32 SupportsFormat(const FString& InCodecFormat, bool bInEncoder, const TMap<FString, FVariant>& InOptions) const override
+	int32 SupportsFormat(TMap<FString, FVariant>& OutFormatInfo, const FString& InCodecFormat, bool bInEncoder, const TMap<FString, FVariant>& InOptions) const override
 	{
 		// Encoder? Not supported here!
 		if (bInEncoder)
@@ -104,7 +104,6 @@ public:
 
 	TSharedPtr<IElectraDecoder, ESPMode::ThreadSafe> CreateDecoderForFormat(const FString& InCodecFormat, const TMap<FString, FVariant>& InOptions, TSharedPtr<IElectraDecoderResourceDelegate, ESPMode::ThreadSafe> InResourceDelegate) override
 	{
-		check(SupportsFormat(InCodecFormat, false, InOptions));
 		return IElectraVideoDecoderH265_DX::Create(InOptions, InResourceDelegate);
 	}
 
@@ -149,7 +148,7 @@ namespace IElectraVideoDecoderH265_DX_Platform
 		{ return Ctx.DxDeviceContext.GetReference(); }
 		virtual bool IsSoftware() const override
 		{ return bIsSW; }
-		
+
 		void Close()
 		{
 			MFT.SafeRelease();
diff --git a/Engine/Plugins/Media/ElectraCodecs/Source/ElectraDecoders/Public/IElectraDecoderFeaturesAndOptions.h b/Engine/Plugins/Media/ElectraCodecs/Source/ElectraDecoders/Public/IElectraDecoderFeaturesAndOptions.h
--- a/Engine/Plugins/Media/ElectraCodecs/Source/ElectraDecoders/Public/IElectraDecoderFeaturesAndOptions.h
+++ b/Engine/Plugins/Media/ElectraCodecs/Source/ElectraDecoders/Public/IElectraDecoderFeaturesAndOptions.h
@@ -158,3 +158,27 @@ namespace IElectraDecoderOption
 	 */
 	const TCHAR* const DecoderConfigurationRecord = TEXT("dcr");
 }
+
+
+
+/**
+ * This namespace contains constant names describe an decoder format.
+ */
+namespace IElectraDecoderFormatInfo
+{
+	/**
+	 * FVariant(bool) : true - The format uses only keyframes.
+	 *
+	 * Indicates that this format encodes every frame as a keyframe and no delta frames.
+	 * This is used with video decoders only. Audio is assumed to be keyframe-only.
+	 * The absence of this on a video decoder or a value of 'false' implies that delta frames are used.
+	 */
+	const TCHAR* const IsEveryFrameKeyframe = TEXT("is_only_keyframes");
+
+	/**
+	 * FVariant(String)
+	 *
+	 * A human readable string of the format to be decoded.
+	 */
+	const TCHAR* const HumanReadableFormatName = TEXT("human_readable_format_name");
+}
diff --git a/Engine/Plugins/Media/ElectraPlayer/Source/ElectraPlayerRuntime/Private/Runtime/Decoder/AudioDecoder.cpp b/Engine/Plugins/Media/ElectraPlayer/Source/ElectraPlayerRuntime/Private/Runtime/Decoder/AudioDecoder.cpp
--- a/Engine/Plugins/Media/ElectraPlayer/Source/ElectraPlayerRuntime/Private/Runtime/Decoder/AudioDecoder.cpp
+++ b/Engine/Plugins/Media/ElectraPlayer/Source/ElectraPlayerRuntime/Private/Runtime/Decoder/AudioDecoder.cpp
@@ -257,7 +257,8 @@ TSharedPtr<IElectraCodecFactory, ESPMode::ThreadSafe> FAudioDecoderImpl::GetDeco
 	OutAddtlCfg.Add(TEXT("codec_name"), FVariant(OutFormat));
 	OutAddtlCfg.Add(TEXT("codec_4cc"), FVariant((uint32)InCodecInfo.GetCodec4CC()));
 	InCodecInfo.GetExtras().ConvertTo(OutAddtlCfg, TEXT("$"));
-	return FactoryModule->GetBestFactoryForFormat(OutFormat, false, OutAddtlCfg);
+	TMap<FString, FVariant> FormatInfo;
+	return FactoryModule->GetBestFactoryForFormat(FormatInfo, OutFormat, false, OutAddtlCfg);
 }
 
 
diff --git a/Engine/Plugins/Media/ElectraPlayer/Source/ElectraPlayerRuntime/Private/Runtime/Decoder/VideoDecoder.cpp b/Engine/Plugins/Media/ElectraPlayer/Source/ElectraPlayerRuntime/Private/Runtime/Decoder/VideoDecoder.cpp
--- a/Engine/Plugins/Media/ElectraPlayer/Source/ElectraPlayerRuntime/Private/Runtime/Decoder/VideoDecoder.cpp
+++ b/Engine/Plugins/Media/ElectraPlayer/Source/ElectraPlayerRuntime/Private/Runtime/Decoder/VideoDecoder.cpp
@@ -292,7 +292,8 @@ TSharedPtr<IElectraCodecFactory, ESPMode::ThreadSafe> FVideoDecoderImpl::GetDeco
 	OutAddtlCfg.Add(TEXT("codec_name"), FVariant(OutFormat));
 	OutAddtlCfg.Add(TEXT("codec_4cc"), FVariant((uint32) InCodecInfo.GetCodec4CC()));
 	InCodecInfo.GetExtras().ConvertTo(OutAddtlCfg, TEXT("$"));
-	return FactoryModule->GetBestFactoryForFormat(OutFormat, false, OutAddtlCfg);
+	TMap<FString, FVariant> FormatInfo;
+	return FactoryModule->GetBestFactoryForFormat(FormatInfo, OutFormat, false, OutAddtlCfg);
 }
 
 
diff --git a/Engine/Plugins/Media/FlacDecoderElectra/Source/FlacDecoderElectra/Private/FlacDecoder/ElectraMediaFlacDecoder.cpp b/Engine/Plugins/Media/FlacDecoderElectra/Source/FlacDecoderElectra/Private/FlacDecoder/ElectraMediaFlacDecoder.cpp
--- a/Engine/Plugins/Media/FlacDecoderElectra/Source/FlacDecoderElectra/Private/FlacDecoder/ElectraMediaFlacDecoder.cpp
+++ b/Engine/Plugins/Media/FlacDecoderElectra/Source/FlacDecoderElectra/Private/FlacDecoder/ElectraMediaFlacDecoder.cpp
@@ -251,13 +251,14 @@ public:
 		OutCodecFactories.Add(AsShared());
 	}
 
-	int32 SupportsFormat(const FString& InCodecFormat, bool bInEncoder, const TMap<FString, FVariant>& InOptions) const override
+	int32 SupportsFormat(TMap<FString, FVariant>& OutFormatInfo, const FString& InCodecFormat, bool bInEncoder, const TMap<FString, FVariant>& InOptions) const override
 	{
 		// Quick check if this is an ask for an encoder or for a 4CC we do not support.
 		if (bInEncoder || !Permitted4CCs.Contains(InCodecFormat))
 		{
 			return 0;
 		}
+		OutFormatInfo.Emplace(IElectraDecoderFormatInfo::HumanReadableFormatName, FVariant(FString(TEXT("Free Lossless Audio Codec (FLAC)"))));
 		return 5;
 	}
 
diff --git a/Engine/Plugins/Media/HAPDecoderElectra/Source/HAPDecoderElectra/Private/HAPDecoder/ElectraMediaHAPDecoder.cpp b/Engine/Plugins/Media/HAPDecoderElectra/Source/HAPDecoderElectra/Private/HAPDecoder/ElectraMediaHAPDecoder.cpp
--- a/Engine/Plugins/Media/HAPDecoderElectra/Source/HAPDecoderElectra/Private/HAPDecoder/ElectraMediaHAPDecoder.cpp
+++ b/Engine/Plugins/Media/HAPDecoderElectra/Source/HAPDecoderElectra/Private/HAPDecoder/ElectraMediaHAPDecoder.cpp
@@ -339,6 +339,11 @@ private:
 class FHAPVideoDecoderElectraFactory : public IElectraCodecFactory, public IElectraCodecModularFeature, public TSharedFromThis<FHAPVideoDecoderElectraFactory, ESPMode::ThreadSafe>
 {
 public:
+	static constexpr uint32 Make4CC(const uint8 A, const uint8 B, const uint8 C, const uint8 D)
+	{
+		return (static_cast<uint32>(A) << 24) | (static_cast<uint32>(B) << 16) | (static_cast<uint32>(C) << 8) | static_cast<uint32>(D);
+	}
+
 	virtual ~FHAPVideoDecoderElectraFactory()
 	{ }
 
@@ -347,9 +352,51 @@ public:
 		OutCodecFactories.Add(AsShared());
 	}
 
-	int32 SupportsFormat(const FString& InCodecFormat, bool bInEncoder, const TMap<FString, FVariant>& InOptions) const override
+	int32 SupportsFormat(TMap<FString, FVariant>& OutFormatInfo, const FString& InCodecFormat, bool bInEncoder, const TMap<FString, FVariant>& InOptions) const override
 	{
-		return bInEncoder || !Permitted4CCs.Contains(InCodecFormat) ? 0 : 1;
+		if (bInEncoder || !Permitted4CCs.Contains(InCodecFormat))
+		{
+			return 0;
+		}
+		switch((uint32)ElectraDecodersUtil::GetVariantValueSafeU64(InOptions, TEXT("codec_4cc"), 0))
+		{
+			case Make4CC('H','a','p','1'):
+			{
+				OutFormatInfo.Emplace(IElectraDecoderFormatInfo::HumanReadableFormatName, FVariant(FString(TEXT("Hap"))));
+				break;
+			}
+			case Make4CC('H','a','p','5'):
+			{
+				OutFormatInfo.Emplace(IElectraDecoderFormatInfo::HumanReadableFormatName, FVariant(FString(TEXT("Hap Alpha"))));
+				break;
+			}
+			case Make4CC('H','a','p','Y'):
+			{
+				OutFormatInfo.Emplace(IElectraDecoderFormatInfo::HumanReadableFormatName, FVariant(FString(TEXT("Hap Q"))));
+				break;
+			}
+			case Make4CC('H','a','p','M'):
+			{
+				OutFormatInfo.Emplace(IElectraDecoderFormatInfo::HumanReadableFormatName, FVariant(FString(TEXT("Hap Q Alpha"))));
+				break;
+			}
+			case Make4CC('H','a','p','7'):
+			{
+				OutFormatInfo.Emplace(IElectraDecoderFormatInfo::HumanReadableFormatName, FVariant(FString(TEXT("Hap R"))));
+				break;
+			}
+			case Make4CC('H','a','p','H'):
+			{
+				OutFormatInfo.Emplace(IElectraDecoderFormatInfo::HumanReadableFormatName, FVariant(FString(TEXT("Hap HDR"))));
+				break;
+			}
+			default:
+			{
+				break;
+			}
+		}
+		OutFormatInfo.Emplace(IElectraDecoderFormatInfo::IsEveryFrameKeyframe, FVariant(true));
+		return 1;
 	}
 
 	void GetConfigurationOptions(TMap<FString, FVariant>& OutOptions) const override
diff --git a/Engine/Plugins/Media/OpusDecoderElectra/Source/OpusDecoderElectra/Private/OpusDecoder/ElectraMediaOpusDecoder.cpp b/Engine/Plugins/Media/OpusDecoderElectra/Source/OpusDecoderElectra/Private/OpusDecoder/ElectraMediaOpusDecoder.cpp
--- a/Engine/Plugins/Media/OpusDecoderElectra/Source/OpusDecoderElectra/Private/OpusDecoder/ElectraMediaOpusDecoder.cpp
+++ b/Engine/Plugins/Media/OpusDecoderElectra/Source/OpusDecoderElectra/Private/OpusDecoder/ElectraMediaOpusDecoder.cpp
@@ -239,7 +239,7 @@ private:
 	// Input configuration
 	FDOpsConfig DOpsConfig;
 	bool bHaveParseddOps = false;
-	
+
 	// Output
 	TArray<IElectraDecoderAudioOutput::EChannelPosition> OutputChannelMap;
 };
@@ -260,13 +260,14 @@ public:
 		OutCodecFactories.Add(AsShared());
 	}
 
-	int32 SupportsFormat(const FString& InCodecFormat, bool bInEncoder, const TMap<FString, FVariant>& InOptions) const override
+	int32 SupportsFormat(TMap<FString, FVariant>& OutFormatInfo, const FString& InCodecFormat, bool bInEncoder, const TMap<FString, FVariant>& InOptions) const override
 	{
 		// Quick check if this is an ask for an encoder or for a 4CC we do not support.
 		if (bInEncoder || !Permitted4CCs.Contains(InCodecFormat))
 		{
 			return 0;
 		}
+		OutFormatInfo.Emplace(IElectraDecoderFormatInfo::HumanReadableFormatName, FVariant(FString(TEXT("Opus"))));
 		return 5;
 	}
 
diff --git a/Engine/Plugins/Media/VPxDecoderElectra/Source/VPxDecoderElectra/Private/VPxDecoder/ElectraMediaVPxDecoder.cpp b/Engine/Plugins/Media/VPxDecoderElectra/Source/VPxDecoderElectra/Private/VPxDecoder/ElectraMediaVPxDecoder.cpp
--- a/Engine/Plugins/Media/VPxDecoderElectra/Source/VPxDecoderElectra/Private/VPxDecoder/ElectraMediaVPxDecoder.cpp
+++ b/Engine/Plugins/Media/VPxDecoderElectra/Source/VPxDecoderElectra/Private/VPxDecoder/ElectraMediaVPxDecoder.cpp
@@ -310,7 +310,7 @@ public:
 		OutCodecFactories.Add(AsShared());
 	}
 
-	int32 SupportsFormat(const FString& InCodecFormat, bool bInEncoder, const TMap<FString, FVariant>& InOptions) const override
+	int32 SupportsFormat(TMap<FString, FVariant>& OutFormatInfo, const FString& InCodecFormat, bool bInEncoder, const TMap<FString, FVariant>& InOptions) const override
 	{
 		if (bInEncoder)
 		{
@@ -836,7 +836,7 @@ FVideoDecoderVPxElectra::EConvertResult FVideoDecoderVPxElectra::ConvertDecoderO
 	TSharedPtr<FVideoDecoderOutputVPxElectra, ESPMode::ThreadSafe> NewOutput = MakeShared<FVideoDecoderOutputVPxElectra>();
 	NewOutput->PTS = In->AccessUnit.PTS;
 	NewOutput->UserValue = In->AccessUnit.UserValue;
-	
+
 	NewOutput->Width = InDecodedImage->d_w;
 	NewOutput->Height = InDecodedImage->d_h;
 	NewOutput->NumBits = InDecodedImage->bit_depth;
