disable ProRes plugins on Windows Arm64 - not supported

#rnx

diff --git a/Engine/Plugins/Media/AppleProResDecoderElectra/AppleProResDecoderElectra.uplugin b/Engine/Plugins/Media/AppleProResDecoderElectra/AppleProResDecoderElectra.uplugin
--- a/Engine/Plugins/Media/AppleProResDecoderElectra/AppleProResDecoderElectra.uplugin
+++ b/Engine/Plugins/Media/AppleProResDecoderElectra/AppleProResDecoderElectra.uplugin
@@ -18,6 +18,7 @@
       "Type": "Runtime",
       "LoadingPhase": "PostEngineInit",
       "PlatformAllowList": [ "Win64", "Mac" ],
+      "PlatformArchitectureDenyList": [ "Win64:arm64" ],
       "TargetDenyList": [ "Server" ]
     }
   ],
diff --git a/Engine/Plugins/Media/AppleProResMedia/AppleProResMedia.uplugin b/Engine/Plugins/Media/AppleProResMedia/AppleProResMedia.uplugin
--- a/Engine/Plugins/Media/AppleProResMedia/AppleProResMedia.uplugin
+++ b/Engine/Plugins/Media/AppleProResMedia/AppleProResMedia.uplugin
@@ -18,6 +18,7 @@
 			"Type": "Runtime",
 			"LoadingPhase": "PostEngineInit",
 			"PlatformAllowList": [ "Win64", "Mac" ],
+			"PlatformArchitectureDenyList": [ "Win64:arm64" ],
 			"TargetDenyList": [ "Server" ]
 		}
 	],
