ElectraCodecs: better error handling in the ProRes decoder when DX12 buffers can't be allocated

#rnx

diff --git a/Engine/Plugins/Media/AppleProResDecoderElectra/Source/AppleProResDecoderElectra/Private/ProResDecoder/Windows/ElectraMediaProResDecoder.cpp b/Engine/Plugins/Media/AppleProResDecoderElectra/Source/AppleProResDecoderElectra/Private/ProResDecoder/Windows/ElectraMediaProResDecoder.cpp
--- a/Engine/Plugins/Media/AppleProResDecoderElectra/Source/AppleProResDecoderElectra/Private/ProResDecoder/Windows/ElectraMediaProResDecoder.cpp
+++ b/Engine/Plugins/Media/AppleProResDecoderElectra/Source/AppleProResDecoderElectra/Private/ProResDecoder/Windows/ElectraMediaProResDecoder.cpp
@@ -735,50 +735,64 @@ IElectraDecoder::EDecoderError FVideoDecoderProResElectra::DecodeAccessUnit(cons
 		switch(OutputPixelFormat)
 		{
 			case kPRFormat_2vuy:
+			{
 				NewOutput->NumBits = 8;
 				NewOutput->BufferFormat = EPixelFormat::PF_B8G8R8A8;
 				NewOutput->BufferEncoding = EElectraTextureSamplePixelEncoding::CbY0CrY1;
 				NewOutput->Width /= 2;
 				NewOutput->Pitch *= 2;
 				break;
+			}
 			case kPRFormat_y416:
+			{
 				NewOutput->NumBits = 16;
 				NewOutput->BufferFormat = EPixelFormat::PF_A16B16G16R16;
 				NewOutput->BufferEncoding = EElectraTextureSamplePixelEncoding::YCbCr_Alpha;
 				NewOutput->Pitch *= 8;
 				break;
+			}
 			case kPRFormat_r4fl:
+			{
 				NewOutput->NumBits = 32;
 				NewOutput->BufferFormat = EPixelFormat::PF_A32B32G32R32F;
 				NewOutput->BufferEncoding = EElectraTextureSamplePixelEncoding::YCbCr_Alpha;
 				NewOutput->Pitch *= 16;
 				break;
+			}
 			case kPRFormat_v210:
+			{
 				NewOutput->NumBits = 10;
 				NewOutput->BufferFormat = EPixelFormat::PF_A2B10G10R10;
 				NewOutput->BufferEncoding = EElectraTextureSamplePixelEncoding::CbY0CrY1;
 				NewOutput->Width = 4 * ((NewOutput->Width + 5) / 6); // each 4 pixel contain 6 horizontally adjacent YCbCr pixels (incl. 4x 2-bit padding)
 				NewOutput->Pitch = NewOutput->Width * 4;
 				break;
+			}
 			case kPRFormat_v216:
+			{
 				NewOutput->NumBits = 16;
 				NewOutput->BufferFormat = EPixelFormat::PF_A16B16G16R16;
 				NewOutput->BufferEncoding = EElectraTextureSamplePixelEncoding::CbY0CrY1;
 				NewOutput->Width /= 2;
 				NewOutput->Pitch *= 4;
 				break;
+			}
 			case kPRFormat_b64a:
+			{
 				NewOutput->NumBits = 16;
 				NewOutput->BufferFormat = EPixelFormat::PF_A16B16G16R16;
 				NewOutput->BufferEncoding = EElectraTextureSamplePixelEncoding::ARGB_BigEndian;
 				NewOutput->Pitch *= 8;
 				break;
+			}
 			case kPRFormat_R10k:
 			case kPRFormat_r210:
 			default:
+			{
 				NewOutput->NumBits = 0;
 				NewOutput->BufferFormat = EPixelFormat::PF_Unknown;
 				break;
+			}
 		}
 
 		NewOutput->ExtraValues.Emplace(TEXT("codec"), FVariant(TEXT("prores")));
@@ -808,17 +822,30 @@ IElectraDecoder::EDecoderError FVideoDecoderProResElectra::DecodeAccessUnit(cons
 			// Request resource and fence...
 			uint32 BufferPitch;
 			D3D12ResourcePool->AllocateOutputDataAsBuffer(NewOutput->GPUBuffer, BufferPitch);
-
-			// Correct output pitch to what the resource is setup for
-			NewOutput->Pitch = BufferPitch;
-
-			void* BufferAddr;
-			HRESULT Res = NewOutput->GPUBuffer.Resource->Map(0, nullptr, &BufferAddr);
-			check(SUCCEEDED(Res));
-
-			check((uint32)OutputBufferBytesPerRow <= BufferPitch);
-			pb.baseAddr = (unsigned char*)BufferAddr;
-			pb.rowBytes = BufferPitch;
+			if (NewOutput->GPUBuffer.Resource)
+			{
+				// Correct output pitch to what the resource is setup for
+				NewOutput->Pitch = BufferPitch;
+
+				void* BufferAddr;
+				HRESULT Res = NewOutput->GPUBuffer.Resource->Map(0, nullptr, &BufferAddr);
+				if (SUCCEEDED(Res))
+				{
+					check((uint32)OutputBufferBytesPerRow <= BufferPitch);
+					pb.baseAddr = (unsigned char*)BufferAddr;
+					pb.rowBytes = BufferPitch;
+				}
+				else
+				{
+					PostError(0, TEXT("Failed to map DX12 texture buffer"), ERRCODE_INTERNAL_FAILED_TO_DECODE_INPUT);
+					return IElectraDecoder::EDecoderError::Error;
+				}
+			}
+			else
+			{
+				PostError(0, TEXT("Failed to allocate DX12 texture buffer"), ERRCODE_INTERNAL_FAILED_TO_DECODE_INPUT);
+				return IElectraDecoder::EDecoderError::Error;
+			}
 		}
 		else
 #endif
