Enabling the ProRes decoder plugin for Electra on Mac

#rnx

diff --git a/Engine/Plugins/Media/AppleProResDecoderElectra/AppleProResDecoderElectra.uplugin b/Engine/Plugins/Media/AppleProResDecoderElectra/AppleProResDecoderElectra.uplugin
--- a/Engine/Plugins/Media/AppleProResDecoderElectra/AppleProResDecoderElectra.uplugin
+++ b/Engine/Plugins/Media/AppleProResDecoderElectra/AppleProResDecoderElectra.uplugin
@@ -11,7 +11,7 @@
   "CanContainContent": false,
   "IsBetaVersion": false,
   "Installed": false,
-  "SupportedTargetPlatforms": [ "Win64" ],
+  "SupportedTargetPlatforms": [ "Win64", "Mac" ],
   "Modules": [
     {
       "Name": "AppleProResDecoderElectra",
